RC solve the protocol performance transactiion issues

diff --git a/Engine/Plugins/VirtualProduction/RemoteControl/Source/RemoteControl/Private/RemoteControlModule.cpp b/Engine/Plugins/VirtualProduction/RemoteControl/Source/RemoteControl/Private/RemoteControlModule.cpp
--- a/Engine/Plugins/VirtualProduction/RemoteControl/Source/RemoteControl/Private/RemoteControlModule.cpp
+++ b/Engine/Plugins/VirtualProduction/RemoteControl/Source/RemoteControl/Private/RemoteControlModule.cpp
@@ -466,15 +466,15 @@ public:
 			UStruct* ContainerType = ObjectAccess.ContainerType.Get();
 
 #if WITH_EDITOR
-			bool bGenerateTransaction = ObjectAccess.Access == ERCAccess::WRITE_TRANSACTION_ACCESS;
-			FScopedTransaction Transaction(LOCTEXT("RemoteSetPropertyTransaction", "Remote Set Object Property"), bGenerateTransaction);
-
+			const bool bGenerateTransaction = ObjectAccess.Access == ERCAccess::WRITE_TRANSACTION_ACCESS;
 			if (bGenerateTransaction)
 			{
-				FEditPropertyChain PreEditChain;
-				ObjectAccess.PropertyPathInfo.ToEditPropertyChain(PreEditChain);
-				Object->PreEditChange(PreEditChain);
+				GEditor->BeginTransaction(LOCTEXT("RemoteSetPropertyTransaction", "Remote Set Object Property"));
 			}
+
+			FEditPropertyChain PreEditChain;
+			ObjectAccess.PropertyPathInfo.ToEditPropertyChain(PreEditChain);
+			Object->PreEditChange(PreEditChain);
 #endif
 
 			FStructDeserializerPolicies Policies;
@@ -508,14 +508,15 @@ public:
 				bSuccess = FStructDeserializer::Deserialize(ObjectAccess.ContainerAdress, *ContainerType, Backend, Policies);
 			}
 
-			// if we are generating a transaction, also generate post edit property event, event if the change ended up unsuccessful
-			// this is to match the pre edit change call that can unregister components for example
+			// Generate post edit property event independently from a transaction
 #if WITH_EDITOR
 			if (bGenerateTransaction)
 			{
-				FPropertyChangedEvent PropertyEvent = ObjectAccess.PropertyPathInfo.ToPropertyChangedEvent();
-				Object->PostEditChangeProperty(PropertyEvent);
+				GEditor->EndTransaction();
 			}
+			
+			FPropertyChangedEvent PropertyEvent = ObjectAccess.PropertyPathInfo.ToPropertyChangedEvent();
+			Object->PostEditChangeProperty(PropertyEvent);
 #endif
 			return bSuccess;
 		}
diff --git a/Engine/Plugins/VirtualProduction/RemoteControl/Source/RemoteControlProtocol/Private/RemoteControlProtocol.cpp b/Engine/Plugins/VirtualProduction/RemoteControl/Source/RemoteControlProtocol/Private/RemoteControlProtocol.cpp
--- a/Engine/Plugins/VirtualProduction/RemoteControl/Source/RemoteControlProtocol/Private/RemoteControlProtocol.cpp
+++ b/Engine/Plugins/VirtualProduction/RemoteControl/Source/RemoteControlProtocol/Private/RemoteControlProtocol.cpp
@@ -6,8 +6,20 @@
 #include "RemoteControlProtocolBinding.h"
 #include "RemoteControlProtocolModule.h"
 
+#include "Misc/CoreDelegates.h"
 #include "UObject/StructOnScope.h"
 
+FRemoteControlProtocol::FRemoteControlProtocol(const FName InProtocolName)
+	: ProtocolName(InProtocolName)
+{
+	FCoreDelegates::OnEndFrame.AddRaw(this, &FRemoteControlProtocol::OnEndFrame);
+}
+
+FRemoteControlProtocol::~FRemoteControlProtocol()
+{
+	FCoreDelegates::OnEndFrame.RemoveAll(this);
+}
+
 FRemoteControlProtocolEntityPtr FRemoteControlProtocol::CreateNewProtocolEntity(FProperty* InProperty, URemoteControlPreset* InOwner, FGuid InPropertyId) const
 {
 	FRemoteControlProtocolEntityPtr NewDataPtr = MakeShared<TStructOnScope<FRemoteControlProtocolEntity>>();
@@ -17,6 +29,39 @@ FRemoteControlProtocolEntityPtr FRemoteControlProtocol::CreateNewProtocolEntity(
 	return NewDataPtr;
 }
 
+void FRemoteControlProtocol::QueueValue(const FRemoteControlProtocolEntityPtr InProtocolEntity, const double InProtocolValue)
+{
+	EntityValuesToApply.Add(InProtocolEntity, InProtocolValue);
+}
+
+void FRemoteControlProtocol::OnEndFrame()
+{
+	for (const TPair<const FRemoteControlProtocolEntityPtr, double>& EntityValuesToApplyPair : EntityValuesToApply)
+	{
+		// Check is the Shared ptr and TStructOnScope is valid
+		if (EntityValuesToApplyPair.Key.IsValid() && EntityValuesToApplyPair.Key->IsValid())
+		{
+			FRemoteControlProtocolEntity* ProtocolEntity = EntityValuesToApplyPair.Key->Get();
+			const double ThisFrameValue = EntityValuesToApplyPair.Value;
+			const double* PreviousFrameValuePtr = PreviousTickValuesToApply.Find(EntityValuesToApplyPair.Key);
+			
+			// Check the value from previous frame
+			if (PreviousFrameValuePtr == nullptr || !FMath::IsNearlyEqual(ThisFrameValue, *PreviousFrameValuePtr))
+			{
+				if (!ProtocolEntity->ApplyProtocolValueToProperty(EntityValuesToApplyPair.Value))
+				{
+					// Warn if the the value can't by applied
+					ensureMsgf(false, TEXT("Can't apply property for Protocol %s and PropertyId %s"),
+					           *ProtocolName.ToString(), *ProtocolEntity->GetPropertyId().ToString());
+				}
+			}
+		}
+	}
+
+	// Move the values from this frame to cached map
+	PreviousTickValuesToApply = MoveTemp(EntityValuesToApply);
+}
+
 TFunction<bool(FRemoteControlProtocolEntityWeakPtr InProtocolEntityWeakPtr)> FRemoteControlProtocol::CreateProtocolComparator(FGuid InPropertyId)
 {
 	return [InPropertyId](FRemoteControlProtocolEntityWeakPtr InProtocolEntityWeakPtr) -> bool
diff --git a/Engine/Plugins/VirtualProduction/RemoteControl/Source/RemoteControlProtocol/Public/IRemoteControlProtocol.h b/Engine/Plugins/VirtualProduction/RemoteControl/Source/RemoteControlProtocol/Public/IRemoteControlProtocol.h
--- a/Engine/Plugins/VirtualProduction/RemoteControl/Source/RemoteControlProtocol/Public/IRemoteControlProtocol.h
+++ b/Engine/Plugins/VirtualProduction/RemoteControl/Source/RemoteControlProtocol/Public/IRemoteControlProtocol.h
@@ -53,4 +53,16 @@ public:
 
 	/** Unbind all protocol entities from the protocol */
 	virtual void UnbindAll() = 0;
+
+	/** Build to core engine delegated and called at the end of a frame */
+	virtual void OnEndFrame() {};
+
+	/**
+	 * Queue protocol entity and value to apply for the protocol
+	 * It stores only unique tick entities which should be apply next frame.
+	 * Prevents from applying more then one entity for frame.
+	 * @param InProtocolEntity Protocol Entity
+	 * @param InProtocolValue Protocol Value
+	 */
+	virtual void QueueValue(const FRemoteControlProtocolEntityPtr InProtocolEntity, const double InProtocolValue) = 0;
 };
diff --git a/Engine/Plugins/VirtualProduction/RemoteControl/Source/RemoteControlProtocol/Public/RemoteControlProtocol.h b/Engine/Plugins/VirtualProduction/RemoteControl/Source/RemoteControlProtocol/Public/RemoteControlProtocol.h
--- a/Engine/Plugins/VirtualProduction/RemoteControl/Source/RemoteControlProtocol/Public/RemoteControlProtocol.h
+++ b/Engine/Plugins/VirtualProduction/RemoteControl/Source/RemoteControlProtocol/Public/RemoteControlProtocol.h
@@ -10,8 +10,13 @@
 class REMOTECONTROLPROTOCOL_API FRemoteControlProtocol : public IRemoteControlProtocol
 {
 public:
+	FRemoteControlProtocol(FName InProtocolName);
+	~FRemoteControlProtocol();
+	
 	//~ Begin IRemoteControlProtocol interface
 	virtual FRemoteControlProtocolEntityPtr CreateNewProtocolEntity(FProperty* InProperty, URemoteControlPreset* InOwner, FGuid InPropertyId) const override;
+	virtual void QueueValue(const FRemoteControlProtocolEntityPtr InProtocolEntity, const double InProtocolValue) override;
+	virtual void OnEndFrame() override;
 	//~ End IRemoteControlProtocol interface
 
 	/**
@@ -20,4 +25,18 @@ public:
 	 * @return Lambda function with ProtocolEntityWeakPtr as argument
 	 */
 	static TFunction<bool(FRemoteControlProtocolEntityWeakPtr InProtocolEntityWeakPtr)> CreateProtocolComparator(FGuid InPropertyId);
+
+protected:
+	/** Current Protocol Name */
+	FName ProtocolName;
+
+private:
+	/** Map of the entities and protocol values about to apply */
+	TMap<const FRemoteControlProtocolEntityPtr, double> EntityValuesToApply;
+
+	/**
+	 * Map of the entities and protocol values from previous tick.
+	 * It allows skipping values that are very close to those of the previous frame
+	 */
+	TMap<const FRemoteControlProtocolEntityPtr, double> PreviousTickValuesToApply;
 };
diff --git a/Engine/Plugins/VirtualProduction/RemoteControlProtocolDMX/RemoteControlProtocolDMX.uplugin b/Engine/Plugins/VirtualProduction/RemoteControlProtocolDMX/RemoteControlProtocolDMX.uplugin
--- a/Engine/Plugins/VirtualProduction/RemoteControlProtocolDMX/RemoteControlProtocolDMX.uplugin
+++ b/Engine/Plugins/VirtualProduction/RemoteControlProtocolDMX/RemoteControlProtocolDMX.uplugin
@@ -20,6 +20,11 @@
 			"Name": "RemoteControlProtocolDMX",
 			"Type": "Runtime",
 			"LoadingPhase": "Default"
+		},
+		{
+			"Name": "RemoteControlProtocolDMXEditor",
+			"Type": "Editor",
+			"LoadingPhase": "Default"
 		}
 	],
 	"Plugins": [
diff --git a/Engine/Plugins/VirtualProduction/RemoteControlProtocolDMX/Source/RemoteControlProtocolDMX/Private/RemoteControlProtocolDMX.cpp b/Engine/Plugins/VirtualProduction/RemoteControlProtocolDMX/Source/RemoteControlProtocolDMX/Private/RemoteControlProtocolDMX.cpp
--- a/Engine/Plugins/VirtualProduction/RemoteControlProtocolDMX/Source/RemoteControlProtocolDMX/Private/RemoteControlProtocolDMX.cpp
+++ b/Engine/Plugins/VirtualProduction/RemoteControlProtocolDMX/Source/RemoteControlProtocolDMX/Private/RemoteControlProtocolDMX.cpp
@@ -4,16 +4,15 @@
 
 #include "DMXProtocolSettings.h"
 #include "DMXProtocolTypes.h"
+#include "RemoteControlLogger.h"
 #include "Interfaces/IDMXProtocol.h"
 #include "IO/DMXInputPort.h"
 #include "IO/DMXPortManager.h"
 #include "Library/DMXEntityFixtureType.h"
 
+#define LOCTEXT_NAMESPACE "RemoteControlProtocolDMX"
 
-TStatId FRemoteControlProtocolDMX::GetStatId() const
-{
-	RETURN_QUICK_DECLARE_CYCLE_STAT(FRemoteControlProtocolDMX, STATGROUP_Tickables);
-}
+const FName FRemoteControlProtocolDMX::ProtocolName = TEXT("DMX");
 
 FRemoteControlDMXProtocolEntity::~FRemoteControlDMXProtocolEntity()
 {
@@ -110,7 +109,7 @@ void FRemoteControlProtocolDMX::Unbind(FRemoteControlProtocolEntityPtr InRemoteC
 	ProtocolsBindings.RemoveAllSwap(CreateProtocolComparator(DMXProtocolEntity->GetPropertyId()));
 }
 
-void FRemoteControlProtocolDMX::Tick(float DeltaTime)
+void FRemoteControlProtocolDMX::OnEndFrame()
 {
 	if (!ProtocolsBindings.Num())
 	{
@@ -140,14 +139,17 @@ void FRemoteControlProtocolDMX::Tick(float DeltaTime)
 			// Get universe DMX signal
 			InputPort->GameThreadGetDMXSignal(DMXProtocolEntity->Universe, DMXProtocolEntity->LastSignalPtr);
 			if (DMXProtocolEntity->LastSignalPtr.IsValid())
-			{					
+			{
+				const FDMXSignalSharedPtr& LastSignalPtr = DMXProtocolEntity->LastSignalPtr;
 				const int32 DMXOffset = DMXProtocolEntity->ExtraSetting.StartingChannel - 1;
 				check(DMXOffset >= 0 && DMXOffset < DMX_UNIVERSE_SIZE);
 				
-				ProcessAndApplyProtocolValue(DMXProtocolEntity->LastSignalPtr, DMXOffset, ProtocolEntity);
+				ProcessAndApplyProtocolValue(LastSignalPtr, DMXOffset, ProtocolEntity);
 			}
 		}
 	}
+
+	FRemoteControlProtocol::OnEndFrame();
 }
 
 void FRemoteControlProtocolDMX::ProcessAndApplyProtocolValue(const FDMXSignalSharedPtr& InSignal, int32 InDMXOffset, const FRemoteControlProtocolEntityPtr& InProtocolEntityPtr)
@@ -157,20 +159,23 @@ void FRemoteControlProtocolDMX::ProcessAndApplyProtocolValue(const FDMXSignalSha
 		return;
 	}
 	
-	const uint8* ChannelData = &InSignal->ChannelData[InDMXOffset];	
 	FRemoteControlDMXProtocolEntity* DMXProtocolEntity = InProtocolEntityPtr->CastChecked<FRemoteControlDMXProtocolEntity>();
+	const uint8* ChannelData = &InSignal->ChannelData[InDMXOffset];	
 	const uint8 NumChannelsToOccupy = UDMXEntityFixtureType::NumChannelsToOccupy(DMXProtocolEntity->DataType);
 
 	if(DMXProtocolEntity->CacheDMXBuffer.Num() != NumChannelsToOccupy ||
 		FMemory::Memcmp(DMXProtocolEntity->CacheDMXBuffer.GetData(), ChannelData, NumChannelsToOccupy) != 0)
 	{
 		const uint32 DMXValue = UDMXEntityFixtureType::BytesToInt(DMXProtocolEntity->DataType, DMXProtocolEntity->bUseLSB, ChannelData);
-
-		// Check if the protocol value was successfully applied
-		if(!DMXProtocolEntity->ApplyProtocolValueToProperty(DMXValue))
+		
+#if WITH_EDITOR
+		FRemoteControlLogger::Get().Log(ProtocolName, [&InSignal, DMXValue]
 		{
-			ensure(false);
-		}
+			return FText::Format(LOCTEXT("DMXEventLog","ExternUniverseID {0}, DMXValue {1}"), InSignal->ExternUniverseID, DMXValue);
+		});
+#endif
+	
+		QueueValue(InProtocolEntityPtr, DMXValue);
 
 		// update cached buffer
 		DMXProtocolEntity->CacheDMXBuffer = TArray<uint8>(ChannelData, NumChannelsToOccupy);
@@ -181,3 +186,5 @@ void FRemoteControlProtocolDMX::UnbindAll()
 {
 	ProtocolsBindings.Empty();
 }
+
+#undef LOCTEXT_NAMESPACE
diff --git a/Engine/Plugins/VirtualProduction/RemoteControlProtocolDMX/Source/RemoteControlProtocolDMX/Private/RemoteControlProtocolDMXModule.cpp b/Engine/Plugins/VirtualProduction/RemoteControlProtocolDMX/Source/RemoteControlProtocolDMX/Private/RemoteControlProtocolDMXModule.cpp
--- a/Engine/Plugins/VirtualProduction/RemoteControlProtocolDMX/Source/RemoteControlProtocolDMX/Private/RemoteControlProtocolDMXModule.cpp
+++ b/Engine/Plugins/VirtualProduction/RemoteControlProtocolDMX/Source/RemoteControlProtocolDMX/Private/RemoteControlProtocolDMXModule.cpp
@@ -14,7 +14,7 @@ public:
 	//~ Begin IModuleInterface
 	virtual void StartupModule() override
 	{
-		IRemoteControlProtocolModule::Get().AddProtocol("DMX", MakeShared<FRemoteControlProtocolDMX>());
+		IRemoteControlProtocolModule::Get().AddProtocol(FRemoteControlProtocolDMX::ProtocolName, MakeShared<FRemoteControlProtocolDMX>());
 	}
 	//~ End IModuleInterface
 };
diff --git a/Engine/Plugins/VirtualProduction/RemoteControlProtocolDMX/Source/RemoteControlProtocolDMX/Public/RemoteControlProtocolDMX.h b/Engine/Plugins/VirtualProduction/RemoteControlProtocolDMX/Source/RemoteControlProtocolDMX/Public/RemoteControlProtocolDMX.h
--- a/Engine/Plugins/VirtualProduction/RemoteControlProtocolDMX/Source/RemoteControlProtocolDMX/Public/RemoteControlProtocolDMX.h
+++ b/Engine/Plugins/VirtualProduction/RemoteControlProtocolDMX/Source/RemoteControlProtocolDMX/Public/RemoteControlProtocolDMX.h
@@ -94,23 +94,21 @@ private:
 /**
  * DMX protocol implementation for Remote Control
  */
-class FRemoteControlProtocolDMX : public FRemoteControlProtocol, public FTickableGameObject
+class FRemoteControlProtocolDMX : public FRemoteControlProtocol
 {
 public:
+	FRemoteControlProtocolDMX()
+		: FRemoteControlProtocol(ProtocolName)
+	{}
+	
 	//~ Begin IRemoteControlProtocol interface
 	virtual void Bind(FRemoteControlProtocolEntityPtr InRemoteControlProtocolEntityPtr) override;
 	virtual void Unbind(FRemoteControlProtocolEntityPtr InRemoteControlProtocolEntityPtr) override;
 	virtual void UnbindAll() override;
 	virtual UScriptStruct* GetProtocolScriptStruct() const override { return FRemoteControlDMXProtocolEntity::StaticStruct(); }
+	virtual void OnEndFrame() override;
 	//~ End IRemoteControlProtocol interface
 
-	//~ Begin FTickableGameObject interface
-	virtual void Tick(float DeltaTime) override;
-	virtual ETickableTickType GetTickableTickType() const override { return ETickableTickType::Always; }
-	virtual bool IsTickableInEditor() const override { return true; }
-	virtual TStatId GetStatId() const override;
-	//~ End FTickableGameObject interface
-
 private:
 	/**
 	 * Apply dmx channel data to the bound property, potentially resize the cache buffer
@@ -123,4 +121,8 @@ private:
 private:
 	/** Binding for the DMX protocol */
 	TArray<FRemoteControlProtocolEntityWeakPtr> ProtocolsBindings;
+
+public:
+	/** DMX protocol name */
+	static const FName ProtocolName;
 };
diff --git a/Engine/Plugins/VirtualProduction/RemoteControlProtocolDMX/Source/RemoteControlProtocolDMX/RemoteControlProtocolDMX.Build.cs b/Engine/Plugins/VirtualProduction/RemoteControlProtocolDMX/Source/RemoteControlProtocolDMX/RemoteControlProtocolDMX.Build.cs
--- a/Engine/Plugins/VirtualProduction/RemoteControlProtocolDMX/Source/RemoteControlProtocolDMX/RemoteControlProtocolDMX.Build.cs
+++ b/Engine/Plugins/VirtualProduction/RemoteControlProtocolDMX/Source/RemoteControlProtocolDMX/RemoteControlProtocolDMX.Build.cs
@@ -17,7 +17,6 @@ public class RemoteControlProtocolDMX : ModuleRules
 		PrivateDependencyModuleNames.AddRange(
 			new string[] {
 				"DMXProtocol",
-				"Engine",
 				"RemoteControl",
 				"RemoteControlProtocol",
 			}
diff --git a/Engine/Plugins/VirtualProduction/RemoteControlProtocolDMX/Source/RemoteControlProtocolDMXEditor/Private/RemoteControlProtocolDMXEditorModule.cpp b/Engine/Plugins/VirtualProduction/RemoteControlProtocolDMX/Source/RemoteControlProtocolDMXEditor/Private/RemoteControlProtocolDMXEditorModule.cpp
--- a/Engine/Plugins/VirtualProduction/RemoteControlProtocolDMX/Source/RemoteControlProtocolDMXEditor/Private/RemoteControlProtocolDMXEditorModule.cpp
+++ b/Engine/Plugins/VirtualProduction/RemoteControlProtocolDMX/Source/RemoteControlProtocolDMXEditor/Private/RemoteControlProtocolDMXEditorModule.cpp
@@ -3,9 +3,10 @@
 #include "CoreMinimal.h"
 
 #include "PropertyEditorModule.h"
+#include "Modules/ModuleInterface.h"
+
 #include "RemoteControlProtocolDMX.h"
 #include "RemoteControlProtocolDMXEditorCustomization.h"
-#include "Modules/ModuleInterface.h"
 
 /**
  * Remote control protocol DMX editor that allows have editor functionality for the protocol
@@ -21,20 +22,20 @@ public:
 
 void FRemoteControlProtocolDMXEditorModule::StartupModule()
 {
-	if (UObjectInitialized() && FModuleManager::Get().IsModuleLoaded(TEXT("PropertyEditor")))
-	{
-		FPropertyEditorModule& PropertyEditorModule = FModuleManager::GetModuleChecked<FPropertyEditorModule>("PropertyEditor");
+	FPropertyEditorModule& PropertyEditorModule = FModuleManager::GetModuleChecked<FPropertyEditorModule>("PropertyEditor");
 
-		PropertyEditorModule.RegisterCustomPropertyTypeLayout(
-            FRemoteControlDMXProtocolEntityExtraSetting::StaticStruct()->GetFName(),
-            FOnGetPropertyTypeCustomizationInstance::CreateStatic(FRemoteControlProtocolDMXEditorTypeCustomization::MakeInstance));	
-	}
+	PropertyEditorModule.RegisterCustomPropertyTypeLayout(
+		FRemoteControlDMXProtocolEntityExtraSetting::StaticStruct()->GetFName(),
+		FOnGetPropertyTypeCustomizationInstance::CreateStatic(FRemoteControlProtocolDMXEditorTypeCustomization::MakeInstance));
 }
 
 void FRemoteControlProtocolDMXEditorModule::ShutdownModule()
 {
-	FPropertyEditorModule& PropertyEditorModule = FModuleManager::GetModuleChecked<FPropertyEditorModule>("PropertyEditor");
-	PropertyEditorModule.UnregisterCustomPropertyTypeLayout(FRemoteControlDMXProtocolEntityExtraSetting::StaticStruct()->GetFName());
+	if (UObjectInitialized() && FModuleManager::Get().IsModuleLoaded(TEXT("PropertyEditor")))
+	{
+		FPropertyEditorModule& PropertyEditorModule = FModuleManager::GetModuleChecked<FPropertyEditorModule>("PropertyEditor");
+		PropertyEditorModule.UnregisterCustomPropertyTypeLayout(FRemoteControlDMXProtocolEntityExtraSetting::StaticStruct()->GetFName());
+	}
 }
 
 IMPLEMENT_MODULE(FRemoteControlProtocolDMXEditorModule, RemoteControlProtocolDMXEditorModule);
\ No newline at end of file
diff --git a/Engine/Plugins/VirtualProduction/RemoteControlProtocolMIDI/Source/RemoteControlProtocolMIDI/Private/RemoteControlProtocolMIDI.cpp b/Engine/Plugins/VirtualProduction/RemoteControlProtocolMIDI/Source/RemoteControlProtocolMIDI/Private/RemoteControlProtocolMIDI.cpp
--- a/Engine/Plugins/VirtualProduction/RemoteControlProtocolMIDI/Source/RemoteControlProtocolMIDI/Private/RemoteControlProtocolMIDI.cpp
+++ b/Engine/Plugins/VirtualProduction/RemoteControlProtocolMIDI/Source/RemoteControlProtocolMIDI/Private/RemoteControlProtocolMIDI.cpp
@@ -3,11 +3,14 @@
 #include "RemoteControlProtocolMIDI.h"
 
 #include "MIDIDeviceManager.h"
+#include "RemoteControlLogger.h"
 #include "RemoteControlProtocolMIDIModule.h"
 #include "RemoteControlProtocolMIDISettings.h"
 
 #define LOCTEXT_NAMESPACE "FRemoteControlProtocolMIDI"
 
+const FName FRemoteControlProtocolMIDI::ProtocolName = TEXT("MIDI");
+
 int32 FRemoteControlMIDIDevice::ResolveDeviceId(const TArray<FFoundMIDIDevice>& InFoundDevices)
 {
 	if(InFoundDevices.Num() == 0)
@@ -27,6 +30,7 @@ int32 FRemoteControlMIDIDevice::ResolveDeviceId(const TArray<FFoundMIDIDevice>&
 			return ResolvedDeviceId = INDEX_NONE;
 		}
 
+
 		return ResolvedDeviceId = MIDISettings->DefaultDevice.ResolveDeviceId(InFoundDevices);
 	}
 
@@ -48,10 +52,8 @@ int32 FRemoteControlMIDIDevice::ResolveDeviceId(const TArray<FFoundMIDIDevice>&
 					FoundDeviceId = FoundDevice.DeviceID;
 					break;
 				}
-				else
-				{
-					return ResolvedDeviceId = FoundDevice.DeviceID;	
-				}
+
+				return ResolvedDeviceId = FoundDevice.DeviceID;
 			}
 		}
 
@@ -143,21 +145,32 @@ void FRemoteControlProtocolMIDI::Bind(FRemoteControlProtocolEntityPtr InRemoteCo
 		MIDIDevices.Add(MIDIDeviceId, MIDIDeviceInputController);
 	}
 
-	TMap<int32, TArray<FRemoteControlProtocolEntityWeakPtr>>* EntityBindingsMapPtr = MIDIDeviceBindings.Find(MIDIDeviceInputControllerPtr->Get());
-	if (EntityBindingsMapPtr == nullptr)
+	if (MIDIProtocolEntity->EventType == EMIDIEventType::NoteOn)
 	{
-		EntityBindingsMapPtr = &MIDIDeviceBindings.Add(MIDIDeviceInputControllerPtr->Get());
+		MIDIDeviceBindings_NoteOn.Add(MIDIProtocolEntity->GetPropertyId(), InRemoteControlProtocolEntityPtr);
 	}
-
-	if (TArray<FRemoteControlProtocolEntityWeakPtr>* ProtocolEntityBindingsArrayPtr = EntityBindingsMapPtr->Find(MIDIProtocolEntity->MessageData1))
+	else if (MIDIProtocolEntity->EventType == EMIDIEventType::ChannelAfterTouch)
 	{
-		ProtocolEntityBindingsArrayPtr->Emplace(MoveTemp(InRemoteControlProtocolEntityPtr));
+		MIDIDeviceBindings_ChannelAfterTouch.Add(MIDIProtocolEntity->GetPropertyId(), InRemoteControlProtocolEntityPtr);
 	}
-	else
+	else if (MIDIProtocolEntity->EventType == EMIDIEventType::ControlChange)
 	{
-		TArray<FRemoteControlProtocolEntityWeakPtr> NewEntityBindingsArrayPtr { MoveTemp(InRemoteControlProtocolEntityPtr) };
+		TMap<int32, TArray<FRemoteControlProtocolEntityWeakPtr>>* EntityBindingsMapPtr = MIDIDeviceBindings_ControlChange.Find(MIDIDeviceInputControllerPtr->Get());
+		if (EntityBindingsMapPtr == nullptr)
+		{
+			EntityBindingsMapPtr = &MIDIDeviceBindings_ControlChange.Add(MIDIDeviceInputControllerPtr->Get());
+		}
+
+		if (TArray<FRemoteControlProtocolEntityWeakPtr>* ProtocolEntityBindingsArrayPtr = EntityBindingsMapPtr->Find(MIDIProtocolEntity->MessageData1))
+		{
+			ProtocolEntityBindingsArrayPtr->Emplace(MoveTemp(InRemoteControlProtocolEntityPtr));
+		}
+		else
+		{
+			TArray<FRemoteControlProtocolEntityWeakPtr> NewEntityBindingsArrayPtr { MoveTemp(InRemoteControlProtocolEntityPtr) };
 
-		EntityBindingsMapPtr->Add(MIDIProtocolEntity->MessageData1, MoveTemp(NewEntityBindingsArrayPtr));
+			EntityBindingsMapPtr->Add(MIDIProtocolEntity->MessageData1, MoveTemp(NewEntityBindingsArrayPtr));
+		}	
 	}
 }
 
@@ -169,44 +182,105 @@ void FRemoteControlProtocolMIDI::Unbind(FRemoteControlProtocolEntityPtr InRemote
 	}
 
 	FRemoteControlMIDIProtocolEntity* MIDIProtocolEntity = InRemoteControlProtocolEntityPtr->CastChecked<FRemoteControlMIDIProtocolEntity>();
-	for (TPair<UMIDIDeviceInputController*, TMap<int32, TArray<FRemoteControlProtocolEntityWeakPtr>>>& BindingsPair : MIDIDeviceBindings)
-	{
-		TMap<int32, TArray<FRemoteControlProtocolEntityWeakPtr>>& BindingMap = BindingsPair.Value;
 
-		for (TPair<int32, TArray<FRemoteControlProtocolEntityWeakPtr>>& BindingMapPair : BindingMap)
+	if (MIDIProtocolEntity->EventType  == EMIDIEventType::NoteOn)
+	{
+		MIDIDeviceBindings_NoteOn.Remove(MIDIProtocolEntity->GetPropertyId());
+	}
+	else if (MIDIProtocolEntity->EventType == EMIDIEventType::ChannelAfterTouch)
+	{
+		MIDIDeviceBindings_ChannelAfterTouch.Remove(MIDIProtocolEntity->GetPropertyId());
+	}
+	else if (MIDIProtocolEntity->EventType == EMIDIEventType::ControlChange)
+	{
+		for (TPair<UMIDIDeviceInputController*, TMap<int32, TArray<FRemoteControlProtocolEntityWeakPtr>>>& BindingsPair : MIDIDeviceBindings_ControlChange)
 		{
-			BindingMapPair.Value.RemoveAllSwap(CreateProtocolComparator(MIDIProtocolEntity->GetPropertyId()));
+			TMap<int32, TArray<FRemoteControlProtocolEntityWeakPtr>>& BindingMap = BindingsPair.Value;
+
+			for (TPair<int32, TArray<FRemoteControlProtocolEntityWeakPtr>>& BindingMapPair : BindingMap)
+			{
+				BindingMapPair.Value.RemoveAllSwap(CreateProtocolComparator(MIDIProtocolEntity->GetPropertyId()));
+			}
 		}
 	}
 }
 
 void FRemoteControlProtocolMIDI::OnReceiveEvent(UMIDIDeviceInputController* MIDIDeviceController, int32 Timestamp, int32 Type, int32 Channel, int32 MessageData1, int32 MessageData2)
 {
-	if (const TMap<int32, TArray<FRemoteControlProtocolEntityWeakPtr>>* MIDIMapBindingsPtr = MIDIDeviceBindings.Find(MIDIDeviceController))
+	const EMIDIEventType MIDIEventType = static_cast<EMIDIEventType>(Type);
+
+#if WITH_EDITOR
+	{		
+		FRemoteControlLogger::Get().Log(ProtocolName, [Type, Channel, MessageData1, MessageData2]
+		{
+			const FString TypeName = StaticEnum<EMIDIEventType>()->GetNameStringByValue(Type);
+			return FText::Format(LOCTEXT("MIDIEventLog","Type {0} TypeName {1}, Channel {2}, MessageData1 {3}, MessageData2 {4}"), Type, FText::FromString(TypeName), Channel, MessageData1, MessageData2);
+		});
+	}
+#endif
+
+	if (MIDIEventType == EMIDIEventType::NoteOn)
+	{
+		for (const TPair<FGuid, FRemoteControlProtocolEntityWeakPtr>& MIDIDeviceBindingPair : MIDIDeviceBindings_NoteOn)
+		{
+			if (const TSharedPtr<TStructOnScope<FRemoteControlProtocolEntity>> ProtocolEntityPtr = MIDIDeviceBindingPair.Value.Pin())
+			{
+				const FRemoteControlMIDIProtocolEntity* MIDIProtocolEntity = ProtocolEntityPtr->CastChecked<FRemoteControlMIDIProtocolEntity>();
+				if (MIDIEventType != MIDIProtocolEntity->EventType || Channel != MIDIProtocolEntity->Channel)
+				{
+					continue;
+				}
+
+				QueueValue(ProtocolEntityPtr, MessageData2);
+			}
+		}
+	}
+	else if (MIDIEventType == EMIDIEventType::ChannelAfterTouch)
 	{
-		if (const TArray<FRemoteControlProtocolEntityWeakPtr>* ProtocolEntityArrayPtr = MIDIMapBindingsPtr->Find(MessageData1))
+		for (const TPair<FGuid, FRemoteControlProtocolEntityWeakPtr>& MIDIDeviceBindingPair : MIDIDeviceBindings_ChannelAfterTouch)
 		{
-			for (const FRemoteControlProtocolEntityWeakPtr& ProtocolEntityWeakPtr : *ProtocolEntityArrayPtr)
+			if (const TSharedPtr<TStructOnScope<FRemoteControlProtocolEntity>> ProtocolEntityPtr = MIDIDeviceBindingPair.Value.Pin())
 			{
-				if (const TSharedPtr<TStructOnScope<FRemoteControlProtocolEntity>> ProtocolEntityPtr = ProtocolEntityWeakPtr.Pin())
+				const FRemoteControlMIDIProtocolEntity* MIDIProtocolEntity = ProtocolEntityPtr->CastChecked<FRemoteControlMIDIProtocolEntity>();
+				if (MIDIEventType != MIDIProtocolEntity->EventType || Channel != MIDIProtocolEntity->Channel)
 				{
-					const FRemoteControlMIDIProtocolEntity* MIDIProtocolEntity = ProtocolEntityPtr->CastChecked<FRemoteControlMIDIProtocolEntity>();
+					continue;
+				}
 
-					if (Type != MIDIProtocolEntity->EventType || Channel != MIDIProtocolEntity->Channel)
+				QueueValue(ProtocolEntityPtr, MessageData1);
+			}
+		}
+	}
+	else if (MIDIEventType == EMIDIEventType::ControlChange)
+	{
+		if (const TMap<int32, TArray<FRemoteControlProtocolEntityWeakPtr>>* MIDIMapBindingsPtr = MIDIDeviceBindings_ControlChange.Find(MIDIDeviceController))
+		{
+			if (const TArray<FRemoteControlProtocolEntityWeakPtr>* ProtocolEntityArrayPtr = MIDIMapBindingsPtr->Find(MessageData1))
+			{
+				for (const FRemoteControlProtocolEntityWeakPtr& ProtocolEntityWeakPtr : *ProtocolEntityArrayPtr)
+				{
+					if (const TSharedPtr<TStructOnScope<FRemoteControlProtocolEntity>> ProtocolEntityPtr = ProtocolEntityWeakPtr.Pin())
 					{
-						continue;
-					}
+						const FRemoteControlMIDIProtocolEntity* MIDIProtocolEntity = ProtocolEntityPtr->CastChecked<FRemoteControlMIDIProtocolEntity>();
 
-					MIDIProtocolEntity->ApplyProtocolValueToProperty(MessageData2);
+						if (MIDIEventType != MIDIProtocolEntity->EventType || Channel != MIDIProtocolEntity->Channel)
+						{
+							continue;
+						}
+
+						QueueValue(ProtocolEntityPtr, MessageData2);
+					}
 				}
 			}
-		}
+		}	
 	}
 }
 
 void FRemoteControlProtocolMIDI::UnbindAll()
 {
-	MIDIDeviceBindings.Empty();
+	MIDIDeviceBindings_NoteOn.Empty();
+	MIDIDeviceBindings_ControlChange.Empty();
+	MIDIDeviceBindings_ChannelAfterTouch.Empty();
 	MIDIDevices.Empty();
 }
 
diff --git a/Engine/Plugins/VirtualProduction/RemoteControlProtocolMIDI/Source/RemoteControlProtocolMIDI/Private/RemoteControlProtocolMIDIModule.cpp b/Engine/Plugins/VirtualProduction/RemoteControlProtocolMIDI/Source/RemoteControlProtocolMIDI/Private/RemoteControlProtocolMIDIModule.cpp
--- a/Engine/Plugins/VirtualProduction/RemoteControlProtocolMIDI/Source/RemoteControlProtocolMIDI/Private/RemoteControlProtocolMIDIModule.cpp
+++ b/Engine/Plugins/VirtualProduction/RemoteControlProtocolMIDI/Source/RemoteControlProtocolMIDI/Private/RemoteControlProtocolMIDIModule.cpp
@@ -37,8 +37,8 @@ void FRemoteControlProtocolMIDIModule::StartupModule()
 	}
 #endif // WITH_EDITOR
 
-	IRemoteControlProtocolModule::Get().AddProtocol("MIDI", MakeShared<FRemoteControlProtocolMIDI>());
 	GetOnMIDIDevicesUpdated().AddRaw(this, &FRemoteControlProtocolMIDIModule::HandleMIDIDevicesUpdated);
+	IRemoteControlProtocolModule::Get().AddProtocol(FRemoteControlProtocolMIDI::ProtocolName, MakeShared<FRemoteControlProtocolMIDI>());
 }
 
 void FRemoteControlProtocolMIDIModule::ShutdownModule()
diff --git a/Engine/Plugins/VirtualProduction/RemoteControlProtocolMIDI/Source/RemoteControlProtocolMIDI/Public/RemoteControlProtocolMIDI.h b/Engine/Plugins/VirtualProduction/RemoteControlProtocolMIDI/Source/RemoteControlProtocolMIDI/Public/RemoteControlProtocolMIDI.h
--- a/Engine/Plugins/VirtualProduction/RemoteControlProtocolMIDI/Source/RemoteControlProtocolMIDI/Public/RemoteControlProtocolMIDI.h
+++ b/Engine/Plugins/VirtualProduction/RemoteControlProtocolMIDI/Source/RemoteControlProtocolMIDI/Public/RemoteControlProtocolMIDI.h
@@ -6,6 +6,7 @@
 #include "RemoteControlProtocolBinding.h"
 
 #include "UObject/StrongObjectPtr.h"
+#include "MIDIDeviceController.h"
 
 #include "RemoteControlProtocolMIDI.generated.h"
 
@@ -123,9 +124,8 @@ public:
 	FRemoteControlMIDIDevice Device;
 
 	/** Midi Event type */
-	// @todo: EMIDIEventType
 	UPROPERTY(EditAnywhere, Category = Mapping)
-	int32 EventType = 11;
+	EMIDIEventType EventType = EMIDIEventType::ControlChange;
 
 	/** Midi button event message data id for binding */
 	UPROPERTY(EditAnywhere, Category = Mapping, meta = (DisplayName = "Mapped channel Id"))
@@ -146,6 +146,10 @@ public:
 class FRemoteControlProtocolMIDI : public FRemoteControlProtocol
 {
 public:
+	FRemoteControlProtocolMIDI()
+		: FRemoteControlProtocol(ProtocolName)
+	{}
+	
 	//~ Begin IRemoteControlProtocol interface
 	virtual void Bind(FRemoteControlProtocolEntityPtr InRemoteControlProtocolEntityPtr) override;
 	virtual void Unbind(FRemoteControlProtocolEntityPtr InRemoteControlProtocolEntityPtr) override;
@@ -157,9 +161,19 @@ private:
 	/** On receive MIDI buffer callback */
 	void OnReceiveEvent(UMIDIDeviceInputController* MIDIDeviceController, int32 Timestamp, int32 Type, int32 Channel, int32 MessageData1, int32 MessageData2);
 
-	/** Binding for the MIDI protocol */
-	TMap<UMIDIDeviceInputController*, TMap<int32, TArray<FRemoteControlProtocolEntityWeakPtr>>> MIDIDeviceBindings;
+	/** Binding for ControlChange (11) MIDI protocol */
+	TMap<UMIDIDeviceInputController*, TMap<int32, TArray<FRemoteControlProtocolEntityWeakPtr>>> MIDIDeviceBindings_ControlChange;
+
+	/** Binding for NoteOn (9) MIDI protocol */
+	TMap<FGuid, FRemoteControlProtocolEntityWeakPtr> MIDIDeviceBindings_NoteOn;
+
+	/** Binding for ChannelAfterTouch (13) MIDI protocol */
+	TMap<FGuid, FRemoteControlProtocolEntityWeakPtr> MIDIDeviceBindings_ChannelAfterTouch;
 
 	/** MIDI devices */
 	TMap<int32, TStrongObjectPtr<UMIDIDeviceInputController>> MIDIDevices;
+
+public:
+	/** MIDI protocol name */
+	static const FName ProtocolName;
 };
diff --git a/Engine/Plugins/VirtualProduction/RemoteControlProtocolMIDI/Source/RemoteControlProtocolMIDIEditor/Private/DetailCustomizations/RemoteControlMIDIProtocolEntityCustomization.cpp b/Engine/Plugins/VirtualProduction/RemoteControlProtocolMIDI/Source/RemoteControlProtocolMIDIEditor/Private/DetailCustomizations/RemoteControlMIDIProtocolEntityCustomization.cpp
deleted file mode 100644
diff --git a/Engine/Plugins/VirtualProduction/RemoteControlProtocolMIDI/Source/RemoteControlProtocolMIDIEditor/Private/DetailCustomizations/RemoteControlMIDIProtocolEntityCustomization.h b/Engine/Plugins/VirtualProduction/RemoteControlProtocolMIDI/Source/RemoteControlProtocolMIDIEditor/Private/DetailCustomizations/RemoteControlMIDIProtocolEntityCustomization.h
deleted file mode 100644
diff --git a/Engine/Plugins/VirtualProduction/RemoteControlProtocolMIDI/Source/RemoteControlProtocolMIDIEditor/Private/RemoteControlProtocolMIDIEditorModule.cpp b/Engine/Plugins/VirtualProduction/RemoteControlProtocolMIDI/Source/RemoteControlProtocolMIDIEditor/Private/RemoteControlProtocolMIDIEditorModule.cpp
--- a/Engine/Plugins/VirtualProduction/RemoteControlProtocolMIDI/Source/RemoteControlProtocolMIDIEditor/Private/RemoteControlProtocolMIDIEditorModule.cpp
+++ b/Engine/Plugins/VirtualProduction/RemoteControlProtocolMIDI/Source/RemoteControlProtocolMIDIEditor/Private/RemoteControlProtocolMIDIEditorModule.cpp
@@ -4,7 +4,6 @@
 #include "RemoteControlProtocolMIDI.h"
 #include "RemoteControlProtocolMIDISettings.h"
 #include "DetailCustomizations/RemoteControlMIDIDeviceCustomization.h"
-#include "DetailCustomizations/RemoteControlMIDIProtocolEntityCustomization.h"
 #include "Modules/ModuleManager.h"
 
 #define LOCTEXT_NAMESPACE "FRemoteControlProtocolMIDIEditorModule"
@@ -21,8 +20,6 @@ public:
 		FPropertyEditorModule& PropertyEditorModule = FModuleManager::GetModuleChecked<FPropertyEditorModule>("PropertyEditor");
 		PropertyEditorModule.RegisterCustomPropertyTypeLayout(FRemoteControlMIDIDevice::StaticStruct()->GetFName(),
 			FOnGetPropertyTypeCustomizationInstance::CreateStatic(&FRemoteControlMIDIDeviceCustomization::MakeInstance));
-		PropertyEditorModule.RegisterCustomClassLayout(FRemoteControlMIDIProtocolEntity::StaticStruct()->GetFName(),
-			FOnGetDetailCustomizationInstance::CreateStatic(&FRemoteControlMIDIProtocolEntityCustomization::MakeInstance));
 	}
 
 	virtual void ShutdownModule() override
@@ -31,7 +28,6 @@ public:
 		{
 			FPropertyEditorModule& PropertyEditorModule = FModuleManager::GetModuleChecked<FPropertyEditorModule>("PropertyEditor");
 			PropertyEditorModule.UnregisterCustomPropertyTypeLayout(FRemoteControlMIDIDevice::StaticStruct()->GetFName());
-			PropertyEditorModule.UnregisterCustomClassLayout(FRemoteControlMIDIProtocolEntity::StaticStruct()->GetFName());
 			PropertyEditorModule.NotifyCustomizationModuleChanged();
 		}
 	}
diff --git a/Engine/Plugins/VirtualProduction/RemoteControlProtocolOSC/Source/RemoteControlProtocolOSC/Private/RemoteControlProtocolOSC.cpp b/Engine/Plugins/VirtualProduction/RemoteControlProtocolOSC/Source/RemoteControlProtocolOSC/Private/RemoteControlProtocolOSC.cpp
--- a/Engine/Plugins/VirtualProduction/RemoteControlProtocolOSC/Source/RemoteControlProtocolOSC/Private/RemoteControlProtocolOSC.cpp
+++ b/Engine/Plugins/VirtualProduction/RemoteControlProtocolOSC/Source/RemoteControlProtocolOSC/Private/RemoteControlProtocolOSC.cpp
@@ -2,8 +2,12 @@
 
 #include "RemoteControlProtocolOSC.h"
 
-#include "OSCMessage.h"
+#include "RemoteControlLogger.h"
+
 #include "OSCManager.h"
+#include "OSCMessage.h"
+
+#define LOCTEXT_NAMESPACE "RemoteControlProtocolOSC"
 
 const FName FRemoteControlProtocolOSC::ProtocolName = TEXT("OSC");
 
@@ -35,9 +39,34 @@ void FRemoteControlProtocolOSC::Unbind(FRemoteControlProtocolEntityPtr InRemoteC
 }
 
 void FRemoteControlProtocolOSC::OSCReceivedMessageEvent(const FOSCMessage& Message, const FString& IPAddress, uint16 Port)
-{
-	const TSharedPtr<IOSCPacket>& Packet = Message.GetPacket();
+{	
 	const FOSCAddress& Address = Message.GetAddress();
+	TArray<float> FloatValues;
+	UOSCManager::GetAllFloats(Message, FloatValues);
+
+#if WITH_EDITOR
+	// Lambda executed only in case of log has been enabled.
+	FRemoteControlLogger::Get().Log(ProtocolName, [&IPAddress, &Address, &FloatValues]
+	{
+		FString FloatValuesString;
+		for (int32 FloatValueIndex = 0; FloatValueIndex < FloatValues.Num(); ++FloatValueIndex)
+		{
+			const FString FloatValueString = FString::SanitizeFloat(FloatValues[FloatValueIndex]);
+			if (FloatValueIndex == 0)
+			{
+				FloatValuesString += FloatValueString;
+			}
+			else
+			{
+				FloatValuesString += TEXT(", ") + FloatValueString;
+			}
+		}
+		
+		return FText::Format(
+			LOCTEXT("OSCEventLog", "IPAddress {0}, FullPath {1}, FloatValues {2}"), FText::FromString(IPAddress),
+			FText::FromString(Address.GetFullPath()), FText::FromString(FloatValuesString));
+	});
+#endif
 
 	const TArray<FRemoteControlProtocolEntityWeakPtr>* ProtocolEntityArrayPtr = Bindings.Find(*Address.GetFullPath());
 	if (ProtocolEntityArrayPtr == nullptr)
@@ -45,15 +74,13 @@ void FRemoteControlProtocolOSC::OSCReceivedMessageEvent(const FOSCMessage& Messa
 		return;
 	}
 
-	TArray<float> FloatValues;
-	UOSCManager::GetAllFloats(Message, FloatValues);
 	for (const float FloatValue : FloatValues)
 	{
 		for (const FRemoteControlProtocolEntityWeakPtr& BindingPtr : *ProtocolEntityArrayPtr)
 		{
 			if (const TSharedPtr<TStructOnScope<FRemoteControlProtocolEntity>> ProtocolEntityPtr = BindingPtr.Pin())
 			{
-				(*ProtocolEntityPtr)->ApplyProtocolValueToProperty(FloatValue);
+				QueueValue(ProtocolEntityPtr, FloatValue);
 			}
 		}
 	}
@@ -63,3 +90,5 @@ void FRemoteControlProtocolOSC::UnbindAll()
 {
 	Bindings.Empty();
 }
+
+#undef LOCTEXT_NAMESPACE
diff --git a/Engine/Plugins/VirtualProduction/RemoteControlProtocolOSC/Source/RemoteControlProtocolOSC/Private/RemoteControlProtocolOSC.h b/Engine/Plugins/VirtualProduction/RemoteControlProtocolOSC/Source/RemoteControlProtocolOSC/Private/RemoteControlProtocolOSC.h
--- a/Engine/Plugins/VirtualProduction/RemoteControlProtocolOSC/Source/RemoteControlProtocolOSC/Private/RemoteControlProtocolOSC.h
+++ b/Engine/Plugins/VirtualProduction/RemoteControlProtocolOSC/Source/RemoteControlProtocolOSC/Private/RemoteControlProtocolOSC.h
@@ -38,6 +38,10 @@ public:
 class FRemoteControlProtocolOSC : public FRemoteControlProtocol
 {
 public:
+	FRemoteControlProtocolOSC()
+		: FRemoteControlProtocol(ProtocolName)
+	{}
+	
 	//~ Begin IRemoteControlProtocol interface
 	virtual void Bind(FRemoteControlProtocolEntityPtr InRemoteControlProtocolEntityPtr) override;
 	virtual void Unbind(FRemoteControlProtocolEntityPtr InRemoteControlProtocolEntityPtr) override;
