1. Implement Protocol Specific Dynamic Columns

diff --git a/Engine/Plugins/VirtualProduction/RemoteControlProtocolDMX/Source/RemoteControlProtocolDMX/Private/RemoteControlProtocolDMX.cpp b/Engine/Plugins/VirtualProduction/RemoteControlProtocolDMX/Source/RemoteControlProtocolDMX/Private/RemoteControlProtocolDMX.cpp
--- a/Engine/Plugins/VirtualProduction/RemoteControlProtocolDMX/Source/RemoteControlProtocolDMX/Private/RemoteControlProtocolDMX.cpp
+++ b/Engine/Plugins/VirtualProduction/RemoteControlProtocolDMX/Source/RemoteControlProtocolDMX/Private/RemoteControlProtocolDMX.cpp
@@ -15,12 +15,22 @@
 #if WITH_EDITOR
 #include "IRCProtocolBindingList.h"
 #include "IRemoteControlProtocolWidgetsModule.h"
+#include "Widgets/Input/SNumericEntryBox.h"
 #endif
 
 #define LOCTEXT_NAMESPACE "RemoteControlProtocolDMX"
 
 const FName FRemoteControlProtocolDMX::ProtocolName = TEXT("DMX");
 
+#if WITH_EDITOR
+
+namespace FRemoteControlDMXProtocolColumns
+{
+	static FName Channel = TEXT("Channel");
+	static FName Universe = TEXT("Universe");
+}
+
+#endif // WITH_EDITOR
 
 uint8 FRemoteControlDMXProtocolEntity::GetRangePropertySize() const
 {
@@ -75,6 +85,39 @@ const FString& FRemoteControlDMXProtocolEntity::GetRangePropertyMaxValue() const
 	}
 }
 
+#if WITH_EDITOR
+
+void FRemoteControlDMXProtocolEntity::RegisterWidgets()
+{
+	FRemoteControlProtocolEntity::RegisterWidgets();
+
+	{ // Universe
+		if (!Widgets.Contains(FRemoteControlDMXProtocolColumns::Universe))
+		{
+			TSharedPtr<SWidget> UniverseWidget = SNew(SNumericEntryBox<int32>)
+				.AllowSpin(true)
+				.Value_Lambda([this]() { return ExtraSetting.Universe; })
+				.OnValueCommitted_Lambda([this](int32 NewValue, ETextCommit::Type InCommitType) { ExtraSetting.Universe = NewValue; });
+
+			Widgets.Add(FRemoteControlDMXProtocolColumns::Universe, UniverseWidget);
+		}
+	}
+	
+	{ // Channel
+		if (!Widgets.Contains(FRemoteControlDMXProtocolColumns::Channel))
+		{
+			TSharedPtr<SWidget> ChannelWidget = SNew(SNumericEntryBox<int32>)
+				.AllowSpin(true)
+				.Value_Lambda([this]() { return ExtraSetting.StartingChannel; })
+				.OnValueCommitted_Lambda([this](int32 NewValue, ETextCommit::Type InCommitType) { ExtraSetting.StartingChannel = NewValue; });
+
+			Widgets.Add(FRemoteControlDMXProtocolColumns::Channel, ChannelWidget);
+		}
+	}
+}
+
+#endif // WITH_EDITOR
+
 void FRemoteControlDMXProtocolEntity::Initialize()
 {
 	// Handle port changes
@@ -323,6 +366,20 @@ void FRemoteControlProtocolDMX::ProcessAutoBinding(const FRemoteControlProtocolE
 		}
 	}
 }
+
+void FRemoteControlProtocolDMX::RegisterColumns()
+{
+	FRemoteControlProtocol::RegisterColumns();
+
+	REGISTER_COLUMN(FRemoteControlDMXProtocolColumns::Channel
+		, LOCTEXT("RCPresetChannelColumnHeader", "Channel")
+		, ProtocolColumnConstants::ColumnSizeMicro);
+
+	REGISTER_COLUMN(FRemoteControlDMXProtocolColumns::Universe
+		, LOCTEXT("RCPresetUniverseColumnHeader", "Universe")
+		, ProtocolColumnConstants::ColumnSizeMicro);
+}
+
 #endif
 
 void FRemoteControlProtocolDMX::UnbindAll()
diff --git a/Engine/Plugins/VirtualProduction/RemoteControlProtocolDMX/Source/RemoteControlProtocolDMX/Public/RemoteControlProtocolDMX.h b/Engine/Plugins/VirtualProduction/RemoteControlProtocolDMX/Source/RemoteControlProtocolDMX/Public/RemoteControlProtocolDMX.h
--- a/Engine/Plugins/VirtualProduction/RemoteControlProtocolDMX/Source/RemoteControlProtocolDMX/Public/RemoteControlProtocolDMX.h
+++ b/Engine/Plugins/VirtualProduction/RemoteControlProtocolDMX/Source/RemoteControlProtocolDMX/Public/RemoteControlProtocolDMX.h
@@ -68,6 +68,13 @@ public:
 	virtual FName GetRangePropertyName() const override { return NAME_UInt32Property; }
 	virtual uint8 GetRangePropertySize() const override;
 	virtual const FString& GetRangePropertyMaxValue() const override;
+
+#if WITH_EDITOR
+
+	/** Register(s) all the widgets of this protocol entity. */
+	virtual void RegisterWidgets() override;
+
+#endif // WITH_EDITOR
 	//~ End FRemoteControlProtocolEntity interface
 
 	/** Initialize struct and delegates */
@@ -158,7 +165,13 @@ private:
 	 * @param InProtocolEntityPtr	Protocol entity pointer
 	 */
 	void ProcessAutoBinding(const FRemoteControlProtocolEntityPtr& InProtocolEntityPtr);
-#endif
+
+protected:
+
+	/** Populates protocol specific columns. */
+	virtual void RegisterColumns() override;
+
+#endif // WITH_EDITOR
 
 private:
 	/** Binding for the DMX protocol */
diff --git a/Engine/Plugins/VirtualProduction/RemoteControlProtocolDMX/Source/RemoteControlProtocolDMX/RemoteControlProtocolDMX.Build.cs b/Engine/Plugins/VirtualProduction/RemoteControlProtocolDMX/Source/RemoteControlProtocolDMX/RemoteControlProtocolDMX.Build.cs
--- a/Engine/Plugins/VirtualProduction/RemoteControlProtocolDMX/Source/RemoteControlProtocolDMX/RemoteControlProtocolDMX.Build.cs
+++ b/Engine/Plugins/VirtualProduction/RemoteControlProtocolDMX/Source/RemoteControlProtocolDMX/RemoteControlProtocolDMX.Build.cs
@@ -26,7 +26,10 @@ public class RemoteControlProtocolDMX : ModuleRules
 		{
 			PrivateDependencyModuleNames.AddRange(
 				new string[] {
-					"RemoteControlProtocolWidgets"
+					"InputCore",
+					"RemoteControlProtocolWidgets",
+					"Slate",
+					"SlateCore",
 				}
 			);
 		}
diff --git a/Engine/Plugins/VirtualProduction/RemoteControlProtocolMIDI/Source/RemoteControlProtocolMIDI/Private/RemoteControlProtocolMIDI.cpp b/Engine/Plugins/VirtualProduction/RemoteControlProtocolMIDI/Source/RemoteControlProtocolMIDI/Private/RemoteControlProtocolMIDI.cpp
--- a/Engine/Plugins/VirtualProduction/RemoteControlProtocolMIDI/Source/RemoteControlProtocolMIDI/Private/RemoteControlProtocolMIDI.cpp
+++ b/Engine/Plugins/VirtualProduction/RemoteControlProtocolMIDI/Source/RemoteControlProtocolMIDI/Private/RemoteControlProtocolMIDI.cpp
@@ -10,12 +10,23 @@
 #if WITH_EDITOR
 #include "IRCProtocolBindingList.h"
 #include "IRemoteControlProtocolWidgetsModule.h"
+#include "Widgets/Input/SNumericEntryBox.h"
 #endif
 
 #define LOCTEXT_NAMESPACE "FRemoteControlProtocolMIDI"
 
 const FName FRemoteControlProtocolMIDI::ProtocolName = TEXT("MIDI");
 
+#if WITH_EDITOR
+
+namespace FRemoteControlMIDIProtocolColumns
+{
+	static FName Channel = TEXT("Channel");
+	static FName Identifier = TEXT("Identifier");
+}
+
+#endif // WITH_EDITOR
+
 int32 FRemoteControlMIDIDevice::ResolveDeviceId(const TArray<FFoundMIDIDevice>& InFoundDevices)
 {
 	if(InFoundDevices.Num() == 0)
@@ -132,6 +143,38 @@ bool FRemoteControlMIDIProtocolEntity::IsSame(const FRemoteControlProtocolEntity
 	return false;
 }
 
+#if WITH_EDITOR
+
+void FRemoteControlMIDIProtocolEntity::RegisterWidgets()
+{
+	FRemoteControlProtocolEntity::RegisterWidgets();
+
+	{ // Identifier
+		if (!Widgets.Contains(FRemoteControlMIDIProtocolColumns::Identifier))
+		{
+			TSharedPtr<SWidget> IdentifierWidget = SNew(SNumericEntryBox<int32>)
+				.AllowSpin(true)
+				.Value_Lambda([this]() { return MessageData1; })
+				.OnValueCommitted_Lambda([this](int32 NewValue, ETextCommit::Type InCommitType) { MessageData1 = NewValue; });
+
+			Widgets.Add(FRemoteControlMIDIProtocolColumns::Identifier, IdentifierWidget);
+		}
+	}
+
+	{ // Channel
+		if (!Widgets.Contains(FRemoteControlMIDIProtocolColumns::Channel))
+		{
+			TSharedPtr<SWidget> ChannelWidget = SNew(SNumericEntryBox<int32>)
+				.AllowSpin(true)
+				.Value_Lambda([this]() { return Channel; })
+				.OnValueCommitted_Lambda([this](int32 NewValue, ETextCommit::Type InCommitType) { Channel = NewValue; });
+
+			Widgets.Add(FRemoteControlMIDIProtocolColumns::Channel, ChannelWidget);
+		}
+	}
+}
+
+#endif // WITH_EDITOR
 
 void FRemoteControlProtocolMIDI::Bind(FRemoteControlProtocolEntityPtr InRemoteControlProtocolEntityPtr)
 {
@@ -360,7 +403,21 @@ void FRemoteControlProtocolMIDI::ProcessAutoBinding(EMIDIEventType MIDIEventType
 		}
 	}
 }
-#endif
+
+void FRemoteControlProtocolMIDI::RegisterColumns()
+{
+	FRemoteControlProtocol::RegisterColumns();
+
+	REGISTER_COLUMN(FRemoteControlMIDIProtocolColumns::Channel
+		, LOCTEXT("RCPresetChannelColumnHeader", "Channel")
+		, ProtocolColumnConstants::ColumnSizeMicro);
+	
+	REGISTER_COLUMN(FRemoteControlMIDIProtocolColumns::Identifier
+		, LOCTEXT("RCPresetIdentifierColumnHeader", "ID")
+		, ProtocolColumnConstants::ColumnSizeMicro);
+}
+
+#endif // WITH_EDITOR
 
 void FRemoteControlProtocolMIDI::UnbindAll()
 {
diff --git a/Engine/Plugins/VirtualProduction/RemoteControlProtocolMIDI/Source/RemoteControlProtocolMIDI/Public/RemoteControlProtocolMIDI.h b/Engine/Plugins/VirtualProduction/RemoteControlProtocolMIDI/Source/RemoteControlProtocolMIDI/Public/RemoteControlProtocolMIDI.h
--- a/Engine/Plugins/VirtualProduction/RemoteControlProtocolMIDI/Source/RemoteControlProtocolMIDI/Public/RemoteControlProtocolMIDI.h
+++ b/Engine/Plugins/VirtualProduction/RemoteControlProtocolMIDI/Source/RemoteControlProtocolMIDI/Public/RemoteControlProtocolMIDI.h
@@ -153,6 +153,13 @@ public:
 	* Used to check for duplicate inputs.
 	*/
 	virtual bool IsSame(const FRemoteControlProtocolEntity* InOther) override;
+
+#if WITH_EDITOR
+
+	/** Register(s) all the widgets of this protocol entity. */
+	virtual void RegisterWidgets() override;
+
+#endif // WITH_EDITOR
 };
 
 /**
@@ -191,7 +198,12 @@ private:
 	 * @param	MessageData1		The first part of the MIDI data
 	 */
 	void ProcessAutoBinding(EMIDIEventType MIDIEventType, int32 Channel, int32 MessageData1);
-#endif
+
+protected:
+
+	/** Populates protocol specific columns. */
+	virtual void RegisterColumns() override;
+#endif // WITH_EDITOR
 
 private:
 
diff --git a/Engine/Plugins/VirtualProduction/RemoteControlProtocolMIDI/Source/RemoteControlProtocolMIDI/RemoteControlProtocolMIDI.Build.cs b/Engine/Plugins/VirtualProduction/RemoteControlProtocolMIDI/Source/RemoteControlProtocolMIDI/RemoteControlProtocolMIDI.Build.cs
--- a/Engine/Plugins/VirtualProduction/RemoteControlProtocolMIDI/Source/RemoteControlProtocolMIDI/RemoteControlProtocolMIDI.Build.cs
+++ b/Engine/Plugins/VirtualProduction/RemoteControlProtocolMIDI/Source/RemoteControlProtocolMIDI/RemoteControlProtocolMIDI.Build.cs
@@ -25,7 +25,10 @@ public class RemoteControlProtocolMIDI : ModuleRules
 		{
 			PrivateDependencyModuleNames.AddRange(
 				new string[] {
-					"RemoteControlProtocolWidgets"
+					"InputCore",
+					"RemoteControlProtocolWidgets",
+					"Slate",
+					"SlateCore",
 				}
 			);
 		}
diff --git a/Engine/Plugins/VirtualProduction/RemoteControlProtocolOSC/Source/RemoteControlProtocolOSC/Private/RemoteControlProtocolOSC.cpp b/Engine/Plugins/VirtualProduction/RemoteControlProtocolOSC/Source/RemoteControlProtocolOSC/Private/RemoteControlProtocolOSC.cpp
--- a/Engine/Plugins/VirtualProduction/RemoteControlProtocolOSC/Source/RemoteControlProtocolOSC/Private/RemoteControlProtocolOSC.cpp
+++ b/Engine/Plugins/VirtualProduction/RemoteControlProtocolOSC/Source/RemoteControlProtocolOSC/Private/RemoteControlProtocolOSC.cpp
@@ -7,6 +7,7 @@
 #if WITH_EDITOR
 #include "IRCProtocolBindingList.h"
 #include "IRemoteControlProtocolWidgetsModule.h"
+#include "Widgets/Input/SEditableTextBox.h"
 #endif
 
 #include "OSCManager.h"
@@ -16,6 +17,15 @@
 
 const FName FRemoteControlProtocolOSC::ProtocolName = TEXT("OSC");
 
+#if WITH_EDITOR
+
+namespace FRemoteControlOSCProtocolColumns
+{
+	static FName Address = TEXT("Address");
+}
+
+#endif // WITH_EDITOR
+
 bool FRemoteControlOSCProtocolEntity::IsSame(const FRemoteControlProtocolEntity* InOther)
 {
 	if(const FRemoteControlOSCProtocolEntity* Other = static_cast<const FRemoteControlOSCProtocolEntity*>(InOther))
@@ -26,6 +36,26 @@ bool FRemoteControlOSCProtocolEntity::IsSame(const FRemoteControlProtocolEntity*
 	return false;
 }
 
+#if WITH_EDITOR
+
+void FRemoteControlOSCProtocolEntity::RegisterWidgets()
+{
+	FRemoteControlProtocolEntity::RegisterWidgets();
+
+	{ // Address
+		if (!Widgets.Contains(FRemoteControlOSCProtocolColumns::Address))
+		{
+			TSharedPtr<SWidget> AddressWidget = SNew(SEditableTextBox)
+				.Text_Lambda([this]() { return FText::FromName(PathName); })
+				.OnTextCommitted_Lambda([this](const FText& NewValue, ETextCommit::Type InCommitType) { PathName = *NewValue.ToString(); });
+
+			Widgets.Add(FRemoteControlOSCProtocolColumns::Address, AddressWidget);
+		}
+	}
+}
+
+#endif // WITH_EDITOR
+
 void FRemoteControlProtocolOSC::Bind(FRemoteControlProtocolEntityPtr InRemoteControlProtocolEntityPtr)
 {
 	if (!ensure(InRemoteControlProtocolEntityPtr.IsValid()))
@@ -136,7 +166,17 @@ void FRemoteControlProtocolOSC::ProcessAutoBinding(const FOSCAddress& InAddress)
 		}
 	}
 }
-#endif
+
+void FRemoteControlProtocolOSC::RegisterColumns()
+{
+	FRemoteControlProtocol::RegisterColumns();
+
+	REGISTER_COLUMN(FRemoteControlOSCProtocolColumns::Address
+		, LOCTEXT("RCPresetAddressColumnHeader", "Address")
+		, ProtocolColumnConstants::ColumnSizeMedium);
+}
+
+#endif // WITH_EDITOR
 
 void FRemoteControlProtocolOSC::UnbindAll()
 {
diff --git a/Engine/Plugins/VirtualProduction/RemoteControlProtocolOSC/Source/RemoteControlProtocolOSC/Private/RemoteControlProtocolOSC.h b/Engine/Plugins/VirtualProduction/RemoteControlProtocolOSC/Source/RemoteControlProtocolOSC/Private/RemoteControlProtocolOSC.h
--- a/Engine/Plugins/VirtualProduction/RemoteControlProtocolOSC/Source/RemoteControlProtocolOSC/Private/RemoteControlProtocolOSC.h
+++ b/Engine/Plugins/VirtualProduction/RemoteControlProtocolOSC/Source/RemoteControlProtocolOSC/Private/RemoteControlProtocolOSC.h
@@ -36,6 +36,13 @@ public:
 	* Used to check for duplicate inputs.
 	*/
 	virtual bool IsSame(const FRemoteControlProtocolEntity* InOther) override;
+
+#if WITH_EDITOR
+
+	/** Register(s) all the widgets of this protocol entity. */
+	virtual void RegisterWidgets() override;
+
+#endif // WITH_EDITOR
 };
 
 /**
@@ -64,7 +71,12 @@ public:
 	 * @param InAddress	OSC address structure
 	 */
 	void ProcessAutoBinding(const FOSCAddress& InAddress);
-#endif
+
+protected:
+
+	/** Populates protocol specific columns. */
+	virtual void RegisterColumns() override;
+#endif // WITH_EDITOR
 
 private:
 	/** Map of the OSC bindings */
diff --git a/Engine/Plugins/VirtualProduction/RemoteControlProtocolOSC/Source/RemoteControlProtocolOSC/RemoteControlProtocolOSC.Build.cs b/Engine/Plugins/VirtualProduction/RemoteControlProtocolOSC/Source/RemoteControlProtocolOSC/RemoteControlProtocolOSC.Build.cs
--- a/Engine/Plugins/VirtualProduction/RemoteControlProtocolOSC/Source/RemoteControlProtocolOSC/RemoteControlProtocolOSC.Build.cs
+++ b/Engine/Plugins/VirtualProduction/RemoteControlProtocolOSC/Source/RemoteControlProtocolOSC/RemoteControlProtocolOSC.Build.cs
@@ -27,8 +27,11 @@ public class RemoteControlProtocolOSC : ModuleRules
 			PrivateDependencyModuleNames.AddRange(
 				new string[]
 				{
+					"InputCore",
 				    "RemoteControlProtocolWidgets",
-					"Settings"
+					"Settings",
+					"Slate",
+					"SlateCore",
 				}
 			);
 		}
