Non unity/pch compile fixes

Cmdline: UnrealEditor Win64 debug -nodebuginfo -nolink -allcores -DisableUnity -CppStd=Cpp17 -allmodules -nopch

diff --git a/Engine/Plugins/Enterprise/DatasmithC4DImporter/Source/DatasmithC4DTranslator/Private/DatasmithC4DImporter.cpp b/Engine/Plugins/Enterprise/DatasmithC4DImporter/Source/DatasmithC4DTranslator/Private/DatasmithC4DImporter.cpp
--- a/Engine/Plugins/Enterprise/DatasmithC4DImporter/Source/DatasmithC4DTranslator/Private/DatasmithC4DImporter.cpp
+++ b/Engine/Plugins/Enterprise/DatasmithC4DImporter/Source/DatasmithC4DTranslator/Private/DatasmithC4DImporter.cpp
@@ -19,6 +19,7 @@
 #endif //WITH_EDITOR
 
 #include "AssetRegistry/AssetRegistryModule.h"
+#include "Components.h"
 #include "Curves/RichCurve.h"
 #include "Engine/StaticMesh.h"
 #include "Engine/StaticMeshActor.h"
diff --git a/Engine/Plugins/Experimental/Enterprise/DataprepGeometryOperations/Source/Private/DataprepGeometryOperations.cpp b/Engine/Plugins/Experimental/Enterprise/DataprepGeometryOperations/Source/Private/DataprepGeometryOperations.cpp
--- a/Engine/Plugins/Experimental/Enterprise/DataprepGeometryOperations/Source/Private/DataprepGeometryOperations.cpp
+++ b/Engine/Plugins/Experimental/Enterprise/DataprepGeometryOperations/Source/Private/DataprepGeometryOperations.cpp
@@ -20,6 +20,7 @@
 #include "MeshDescriptionToDynamicMesh.h"
 #include "DynamicMesh/Operations/MergeCoincidentMeshEdges.h"
 #include "PhysicsEngine/BodySetup.h"
+#include "UObject/UObjectIterator.h"
 
 DEFINE_LOG_CATEGORY(LogDataprepGeometryOperations);
 
diff --git a/Engine/Plugins/Experimental/Enterprise/DataprepGeometryOperations/Source/Private/DataprepGeometrySelectionTransforms.cpp b/Engine/Plugins/Experimental/Enterprise/DataprepGeometryOperations/Source/Private/DataprepGeometrySelectionTransforms.cpp
--- a/Engine/Plugins/Experimental/Enterprise/DataprepGeometryOperations/Source/Private/DataprepGeometrySelectionTransforms.cpp
+++ b/Engine/Plugins/Experimental/Enterprise/DataprepGeometryOperations/Source/Private/DataprepGeometrySelectionTransforms.cpp
@@ -1,9 +1,11 @@
 // Copyright Epic Games, Inc. All Rights Reserved.
 
 #include "DataprepGeometrySelectionTransforms.h"
+#include "Engine/Level.h"
 #include "Engine/StaticMesh.h"
 #include "Engine/StaticMeshActor.h"
 #include "Engine/Texture.h"
+#include "Engine/World.h"
 #include "IMeshMergeUtilities.h"
 #include "DataprepOperationsLibraryUtil.h"
 #include "DataprepGeometryOperations.h"
diff --git a/Engine/Plugins/Experimental/GPULightmass/Source/GPULightmass/Private/Scene/Lights.cpp b/Engine/Plugins/Experimental/GPULightmass/Source/GPULightmass/Private/Scene/Lights.cpp
--- a/Engine/Plugins/Experimental/GPULightmass/Source/GPULightmass/Private/Scene/Lights.cpp
+++ b/Engine/Plugins/Experimental/GPULightmass/Source/GPULightmass/Private/Scene/Lights.cpp
@@ -8,6 +8,7 @@
 #include "ReflectionEnvironment.h"
 #include "RectLightTextureManager.h"
 #include "Components/SkyAtmosphereComponent.h"
+#include "TextureResource.h"
 #include "UObject/UObjectIterator.h"
 #include "RenderGraphUtils.h"
 
diff --git a/Engine/Plugins/Experimental/GPULightmass/Source/GPULightmassEditor/Private/GPULightmassEditorModule.cpp b/Engine/Plugins/Experimental/GPULightmass/Source/GPULightmassEditor/Private/GPULightmassEditorModule.cpp
--- a/Engine/Plugins/Experimental/GPULightmass/Source/GPULightmassEditor/Private/GPULightmassEditorModule.cpp
+++ b/Engine/Plugins/Experimental/GPULightmass/Source/GPULightmassEditor/Private/GPULightmassEditorModule.cpp
@@ -5,6 +5,7 @@
 #include "Internationalization/Internationalization.h"
 #include "Modules/ModuleManager.h"
 #include "HAL/IConsoleManager.h"
+#include "PropertyEditorModule.h"
 #include "RenderingThread.h"
 #include "Interfaces/IPluginManager.h"
 #include "ShaderCore.h"
diff --git a/Engine/Plugins/Experimental/HairModelingToolset/Source/HairModelingToolset/Private/GroomCardsEditorTool.cpp b/Engine/Plugins/Experimental/HairModelingToolset/Source/HairModelingToolset/Private/GroomCardsEditorTool.cpp
--- a/Engine/Plugins/Experimental/HairModelingToolset/Source/HairModelingToolset/Private/GroomCardsEditorTool.cpp
+++ b/Engine/Plugins/Experimental/HairModelingToolset/Source/HairModelingToolset/Private/GroomCardsEditorTool.cpp
@@ -2,6 +2,7 @@
 
 #include "GroomCardsEditorTool.h"
 #include "InteractiveToolManager.h"
+#include "SceneManagement.h"
 #include "ToolBuilderUtil.h"
 
 #include "BaseBehaviors/SingleClickBehavior.h"
diff --git a/Engine/Plugins/Experimental/HairModelingToolset/Source/HairModelingToolset/Private/GroomToMeshTool.cpp b/Engine/Plugins/Experimental/HairModelingToolset/Source/HairModelingToolset/Private/GroomToMeshTool.cpp
--- a/Engine/Plugins/Experimental/HairModelingToolset/Source/HairModelingToolset/Private/GroomToMeshTool.cpp
+++ b/Engine/Plugins/Experimental/HairModelingToolset/Source/HairModelingToolset/Private/GroomToMeshTool.cpp
@@ -2,6 +2,7 @@
 
 #include "GroomToMeshTool.h"
 #include "InteractiveToolManager.h"
+#include "Materials/Material.h"
 #include "ToolBuilderUtil.h"
 
 #include "DynamicMesh/DynamicMesh3.h"
diff --git a/Engine/Plugins/Experimental/PanoramicCapture/Source/PanoramicCapture/Private/StereoCapturePawn.cpp b/Engine/Plugins/Experimental/PanoramicCapture/Source/PanoramicCapture/Private/StereoCapturePawn.cpp
--- a/Engine/Plugins/Experimental/PanoramicCapture/Source/PanoramicCapture/Private/StereoCapturePawn.cpp
+++ b/Engine/Plugins/Experimental/PanoramicCapture/Source/PanoramicCapture/Private/StereoCapturePawn.cpp
@@ -5,6 +5,7 @@
 #include "StereoPanoramaManager.h"
 #include "Engine/Texture2D.h"
 #include "Engine/Engine.h"
+#include "TextureResource.h"
 
 void AStereoCapturePawn::UpdateStereoAtlas(UObject* WorldContextObject, struct FLatentActionInfo LatentInfo)
 {
diff --git a/Engine/Plugins/Experimental/WidgetEditorToolPalette/Source/WidgetEditorToolPalette/Public/Settings/CreateWidgetToolSettings.h b/Engine/Plugins/Experimental/WidgetEditorToolPalette/Source/WidgetEditorToolPalette/Public/Settings/CreateWidgetToolSettings.h
--- a/Engine/Plugins/Experimental/WidgetEditorToolPalette/Source/WidgetEditorToolPalette/Public/Settings/CreateWidgetToolSettings.h
+++ b/Engine/Plugins/Experimental/WidgetEditorToolPalette/Source/WidgetEditorToolPalette/Public/Settings/CreateWidgetToolSettings.h
@@ -3,6 +3,7 @@
 #pragma once
 
 #include "CoreMinimal.h"
+#include "Framework/Commands/InputChord.h"
 #include "UObject/ObjectMacros.h"
 #include "UObject/Object.h"
 #include "Engine/DeveloperSettings.h"
diff --git a/Engine/Plugins/Runtime/AR/AppleAR/AppleARKit/Source/AppleARKit/Private/AppleARKitModule.cpp b/Engine/Plugins/Runtime/AR/AppleAR/AppleARKit/Source/AppleARKit/Private/AppleARKitModule.cpp
--- a/Engine/Plugins/Runtime/AR/AppleAR/AppleARKit/Source/AppleARKit/Private/AppleARKitModule.cpp
+++ b/Engine/Plugins/Runtime/AR/AppleAR/AppleARKit/Source/AppleARKit/Private/AppleARKitModule.cpp
@@ -8,6 +8,7 @@
 #include "CoreGlobals.h"
 #include "Misc/CoreDelegates.h"
 #include "Interfaces/IPluginManager.h"
+#include "ShaderCore.h"
 
 #if WITH_EDITOR
 	#include "ISettingsModule.h"
diff --git a/Engine/Plugins/Runtime/AR/AppleAR/AppleARKit/Source/AppleARKit/Private/AppleARKitTextures.cpp b/Engine/Plugins/Runtime/AR/AppleAR/AppleARKit/Source/AppleARKit/Private/AppleARKitTextures.cpp
--- a/Engine/Plugins/Runtime/AR/AppleAR/AppleARKit/Source/AppleARKit/Private/AppleARKitTextures.cpp
+++ b/Engine/Plugins/Runtime/AR/AppleAR/AppleARKit/Source/AppleARKit/Private/AppleARKitTextures.cpp
@@ -12,6 +12,7 @@
 #include "AppleARKitFrame.h"
 #include "HAL/RunnableThread.h"
 #include "GlobalShader.h"
+#include "TextureResource.h"
 
 #if PLATFORM_APPLE
 	#import <Metal/Metal.h>
diff --git a/Engine/Plugins/Runtime/MicrosoftSpatialAudio/Source/MicrosoftSpatialSound/Private/MicrosoftSpatialSoundPlugin.cpp b/Engine/Plugins/Runtime/MicrosoftSpatialAudio/Source/MicrosoftSpatialSound/Private/MicrosoftSpatialSoundPlugin.cpp
--- a/Engine/Plugins/Runtime/MicrosoftSpatialAudio/Source/MicrosoftSpatialSound/Private/MicrosoftSpatialSoundPlugin.cpp
+++ b/Engine/Plugins/Runtime/MicrosoftSpatialAudio/Source/MicrosoftSpatialSound/Private/MicrosoftSpatialSoundPlugin.cpp
@@ -4,6 +4,7 @@
 #include "AudioAnalytics.h"
 #include "Features/IModularFeatures.h"
 #include "HAL/PlatformMisc.h"
+#include "Misc/Paths.h"
 
 DEFINE_LOG_CATEGORY_STATIC(LogMicrosoftSpatialSound, Verbose, All);
 
diff --git a/Engine/Plugins/Runtime/Oculus/OculusAudio/Source/OculusAudio/Private/OculusAudioGeometryComponent.cpp b/Engine/Plugins/Runtime/Oculus/OculusAudio/Source/OculusAudio/Private/OculusAudioGeometryComponent.cpp
--- a/Engine/Plugins/Runtime/Oculus/OculusAudio/Source/OculusAudio/Private/OculusAudioGeometryComponent.cpp
+++ b/Engine/Plugins/Runtime/Oculus/OculusAudio/Source/OculusAudio/Private/OculusAudioGeometryComponent.cpp
@@ -1,5 +1,6 @@
 // Copyright Epic Games, Inc. All Rights Reserved.
 #include "OculusAudioGeometryComponent.h"
+#include "Components/StaticMeshComponent.h"
 #include "OculusAudioMaterialComponent.h"
 #include "OculusAudioContextManager.h"
 #include "OculusAudio.h"
@@ -9,6 +10,7 @@
 #include "Engine/StaticMesh.h"
 #include "Runtime/Core/Public/Serialization/CustomVersion.h"
 #include "Materials/MaterialInstanceDynamic.h"
+#include "StaticMeshResources.h"
 
 
 #define COMPILE_AS_DEBUG_MODE 0
diff --git a/Engine/Plugins/Runtime/Oculus/OculusVR/Source/OculusEditor/Private/OculusToolWidget.cpp b/Engine/Plugins/Runtime/Oculus/OculusVR/Source/OculusEditor/Private/OculusToolWidget.cpp
--- a/Engine/Plugins/Runtime/Oculus/OculusVR/Source/OculusEditor/Private/OculusToolWidget.cpp
+++ b/Engine/Plugins/Runtime/Oculus/OculusVR/Source/OculusEditor/Private/OculusToolWidget.cpp
@@ -12,6 +12,7 @@
 #include "EngineUtils.h"
 #include "Editor.h"
 #include "Styling/AppStyle.h"
+#include "UObject/UObjectIterator.h"
 #include "Widgets/Text/SRichTextBlock.h"
 #include "UObject/EnumProperty.h"
 #include "EdGraph/EdGraph.h"
diff --git a/Engine/Plugins/Runtime/Oculus/OculusVR/Source/OculusHMD/Private/OculusHMD_CustomPresent.cpp b/Engine/Plugins/Runtime/Oculus/OculusVR/Source/OculusHMD/Private/OculusHMD_CustomPresent.cpp
--- a/Engine/Plugins/Runtime/Oculus/OculusVR/Source/OculusHMD/Private/OculusHMD_CustomPresent.cpp
+++ b/Engine/Plugins/Runtime/Oculus/OculusVR/Source/OculusHMD/Private/OculusHMD_CustomPresent.cpp
@@ -4,6 +4,7 @@
 
 #if OCULUS_HMD_SUPPORTED_PLATFORMS
 #include "OculusHMD.h"
+#include "RHIStaticStates.h"
 #include "ScreenRendering.h"
 #include "PipelineStateCache.h"
 #include "ClearQuad.h"
diff --git a/Engine/Plugins/Runtime/Oculus/OculusVR/Source/OculusHMD/Private/OculusHMD_SpectatorScreenController.cpp b/Engine/Plugins/Runtime/Oculus/OculusVR/Source/OculusHMD/Private/OculusHMD_SpectatorScreenController.cpp
--- a/Engine/Plugins/Runtime/Oculus/OculusVR/Source/OculusHMD/Private/OculusHMD_SpectatorScreenController.cpp
+++ b/Engine/Plugins/Runtime/Oculus/OculusVR/Source/OculusHMD/Private/OculusHMD_SpectatorScreenController.cpp
@@ -1,6 +1,7 @@
 // Copyright Epic Games, Inc. All Rights Reserved.
 
 #include "OculusHMD_SpectatorScreenController.h"
+#include "TextureResource.h"
 
 #if OCULUS_HMD_SUPPORTED_PLATFORMS
 #include "OculusHMD.h"
diff --git a/Engine/Plugins/Runtime/Oculus/OculusVR/Source/OculusHMD/Private/OculusHMD_Splash.cpp b/Engine/Plugins/Runtime/Oculus/OculusVR/Source/OculusHMD/Private/OculusHMD_Splash.cpp
--- a/Engine/Plugins/Runtime/Oculus/OculusVR/Source/OculusHMD/Private/OculusHMD_Splash.cpp
+++ b/Engine/Plugins/Runtime/Oculus/OculusVR/Source/OculusHMD/Private/OculusHMD_Splash.cpp
@@ -1,6 +1,7 @@
 // Copyright Epic Games, Inc. All Rights Reserved.
 
 #include "OculusHMD_Splash.h"
+#include "TextureResource.h"
 
 #if OCULUS_HMD_SUPPORTED_PLATFORMS
 #include "OculusHMD.h"
diff --git a/Engine/Plugins/Runtime/Oculus/OculusVR/Source/OculusInput/Private/OculusHandComponent.cpp b/Engine/Plugins/Runtime/Oculus/OculusVR/Source/OculusInput/Private/OculusHandComponent.cpp
--- a/Engine/Plugins/Runtime/Oculus/OculusVR/Source/OculusInput/Private/OculusHandComponent.cpp
+++ b/Engine/Plugins/Runtime/Oculus/OculusVR/Source/OculusInput/Private/OculusHandComponent.cpp
@@ -1,5 +1,6 @@
 // Copyright Epic Games, Inc. All Rights Reserved.
 #include "OculusHandComponent.h"
+#include "GameFramework/Pawn.h"
 #include "OculusInput.h"
 
 #include "Engine/SkeletalMesh.h"
diff --git a/Engine/Plugins/Runtime/Oculus/OculusVR/Source/OculusMR/Private/OculusMR_CastingCameraActor.cpp b/Engine/Plugins/Runtime/Oculus/OculusVR/Source/OculusMR/Private/OculusMR_CastingCameraActor.cpp
--- a/Engine/Plugins/Runtime/Oculus/OculusVR/Source/OculusMR/Private/OculusMR_CastingCameraActor.cpp
+++ b/Engine/Plugins/Runtime/Oculus/OculusVR/Source/OculusMR/Private/OculusMR_CastingCameraActor.cpp
@@ -2,6 +2,7 @@
 
 #include "OculusMR_CastingCameraActor.h"
 
+#include "Materials/Material.h"
 #include "OculusMRPrivate.h"
 #include "OculusHMD_Settings.h"
 #include "OculusHMD.h"
diff --git a/Engine/Plugins/Runtime/Oculus/OculusVR/Source/OculusMR/Private/OculusMR_PlaneMeshComponent.cpp b/Engine/Plugins/Runtime/Oculus/OculusVR/Source/OculusMR/Private/OculusMR_PlaneMeshComponent.cpp
--- a/Engine/Plugins/Runtime/Oculus/OculusVR/Source/OculusMR/Private/OculusMR_PlaneMeshComponent.cpp
+++ b/Engine/Plugins/Runtime/Oculus/OculusVR/Source/OculusMR/Private/OculusMR_PlaneMeshComponent.cpp
@@ -5,6 +5,7 @@
 #include "RenderResource.h"
 #include "PrimitiveViewRelevance.h"
 #include "PrimitiveSceneProxy.h"
+#include "TextureResource.h"
 #include "VertexFactory.h"
 #include "MaterialShared.h"
 #include "Engine/CollisionProfile.h"
diff --git a/Engine/Plugins/Runtime/Steam/SteamAudio/Source/SteamAudio/Private/PhononScene.cpp b/Engine/Plugins/Runtime/Steam/SteamAudio/Source/SteamAudio/Private/PhononScene.cpp
--- a/Engine/Plugins/Runtime/Steam/SteamAudio/Source/SteamAudio/Private/PhononScene.cpp
+++ b/Engine/Plugins/Runtime/Steam/SteamAudio/Source/SteamAudio/Private/PhononScene.cpp
@@ -4,6 +4,7 @@
 
 #include "PhononScene.h"
 
+#include "Model.h"
 #include "PhononCommon.h"
 #include "PhononGeometryComponent.h"
 #include "PhononMaterialComponent.h"
diff --git a/Engine/Plugins/Runtime/Steam/SteamAudio/Source/SteamAudioEditor/Private/PhononSourceComponentDetails.cpp b/Engine/Plugins/Runtime/Steam/SteamAudio/Source/SteamAudioEditor/Private/PhononSourceComponentDetails.cpp
--- a/Engine/Plugins/Runtime/Steam/SteamAudio/Source/SteamAudioEditor/Private/PhononSourceComponentDetails.cpp
+++ b/Engine/Plugins/Runtime/Steam/SteamAudio/Source/SteamAudioEditor/Private/PhononSourceComponentDetails.cpp
@@ -18,6 +18,7 @@
 #include "IDetailChildrenBuilder.h"
 #include "LevelEditorViewport.h"
 #include "EngineUtils.h"
+#include "UObject/UObjectIterator.h"
 #include "Widgets/Input/SButton.h"
 #include "Widgets/Text/STextBlock.h"
 #include "Editor.h"
diff --git a/Engine/Plugins/Runtime/WarpUtils/Source/PFMExporter/Private/PFMExporterShader.cpp b/Engine/Plugins/Runtime/WarpUtils/Source/PFMExporter/Private/PFMExporterShader.cpp
--- a/Engine/Plugins/Runtime/WarpUtils/Source/PFMExporter/Private/PFMExporterShader.cpp
+++ b/Engine/Plugins/Runtime/WarpUtils/Source/PFMExporter/Private/PFMExporterShader.cpp
@@ -3,6 +3,7 @@
 #include "PFMExporterShader.h"
 #include "PFMExporterMesh.h"
 
+#include "Rendering/PositionVertexBuffer.h"
 #include "Shader.h"
 #include "GlobalShader.h"
 #include "ShaderParameters.h"
@@ -18,6 +19,7 @@
 
 #include "Components/StaticMeshComponent.h"
 #include "Components/SceneComponent.h"
+#include "StaticMeshResources.h"
 
 #define PFMExporterShaderFileName TEXT("/Plugin/WarpUtils/Private/PFMExporterShaders.usf")
 
diff --git a/Engine/Plugins/Runtime/nDisplay/Source/DisplayCluster/Private/Render/Synchronization/Windows/DisplayClusterRenderSyncHelperWindows.cpp b/Engine/Plugins/Runtime/nDisplay/Source/DisplayCluster/Private/Render/Synchronization/Windows/DisplayClusterRenderSyncHelperWindows.cpp
--- a/Engine/Plugins/Runtime/nDisplay/Source/DisplayCluster/Private/Render/Synchronization/Windows/DisplayClusterRenderSyncHelperWindows.cpp
+++ b/Engine/Plugins/Runtime/nDisplay/Source/DisplayCluster/Private/Render/Synchronization/Windows/DisplayClusterRenderSyncHelperWindows.cpp
@@ -8,6 +8,7 @@
 #include "Engine/GameViewportClient.h"
 #include "RHI.h"
 #include "DynamicRHI.h"
+#include "UnrealClient.h"
 
 #pragma warning(push)
 #pragma warning (disable : 4005)
diff --git a/Engine/Plugins/Runtime/nDisplay/Source/DisplayCluster/Private/Render/Synchronization/Windows/DisplayClusterRenderSyncPolicyEthernetWindows.cpp b/Engine/Plugins/Runtime/nDisplay/Source/DisplayCluster/Private/Render/Synchronization/Windows/DisplayClusterRenderSyncPolicyEthernetWindows.cpp
--- a/Engine/Plugins/Runtime/nDisplay/Source/DisplayCluster/Private/Render/Synchronization/Windows/DisplayClusterRenderSyncPolicyEthernetWindows.cpp
+++ b/Engine/Plugins/Runtime/nDisplay/Source/DisplayCluster/Private/Render/Synchronization/Windows/DisplayClusterRenderSyncPolicyEthernetWindows.cpp
@@ -10,12 +10,14 @@
 
 #include "Math/UnrealMathUtility.h"
 
+#include "UnrealClient.h"
 #include "Windows/AllowWindowsPlatformTypes.h"
 #include <windows.h>
 #include "DirectX/Include/DXGI.h"
 #include "dwmapi.h"
 #include "Windows/HideWindowsPlatformTypes.h"
 
+#include "RHI.h"
 #include "RHIResources.h"
 
 
diff --git a/Engine/Plugins/Runtime/nDisplay/Source/DisplayClusterStageMonitoring/Private/DWMSyncWatchdog.cpp b/Engine/Plugins/Runtime/nDisplay/Source/DisplayClusterStageMonitoring/Private/DWMSyncWatchdog.cpp
--- a/Engine/Plugins/Runtime/nDisplay/Source/DisplayClusterStageMonitoring/Private/DWMSyncWatchdog.cpp
+++ b/Engine/Plugins/Runtime/nDisplay/Source/DisplayClusterStageMonitoring/Private/DWMSyncWatchdog.cpp
@@ -12,6 +12,7 @@
 #include "Render/IDisplayClusterRenderManager.h"
 #include "Render/Presentation/IDisplayClusterPresentation.h"
 #include "RendererInterface.h"
+#include "UnrealClient.h"
 
 
 #pragma warning(push)
diff --git a/Engine/Plugins/VirtualProduction/RemoteControlProtocolMIDI/Source/RemoteControlProtocolMIDIEditor/Private/DetailCustomizations/RemoteControlMIDIDeviceCustomization.h b/Engine/Plugins/VirtualProduction/RemoteControlProtocolMIDI/Source/RemoteControlProtocolMIDIEditor/Private/DetailCustomizations/RemoteControlMIDIDeviceCustomization.h
--- a/Engine/Plugins/VirtualProduction/RemoteControlProtocolMIDI/Source/RemoteControlProtocolMIDIEditor/Private/DetailCustomizations/RemoteControlMIDIDeviceCustomization.h
+++ b/Engine/Plugins/VirtualProduction/RemoteControlProtocolMIDI/Source/RemoteControlProtocolMIDIEditor/Private/DetailCustomizations/RemoteControlMIDIDeviceCustomization.h
@@ -1,4 +1,4 @@
-﻿// Copyright Epic Games, Inc. All Rights Reserved.
+// Copyright Epic Games, Inc. All Rights Reserved.
 
 #pragma once
 
@@ -12,6 +12,7 @@ class ITableRow;
 class SComboButton;
 class SMIDIDeviceComboBox;
 class STableViewBase;
+class SWidget;
 template <typename ItemType> class SListView;
 
 /** Menu item for a MIDI Device */
diff --git a/Engine/Plugins/VirtualProduction/TextureShare/Source/TextureShare/Private/Game/ViewExtension/TextureShareSceneViewExtension.h b/Engine/Plugins/VirtualProduction/TextureShare/Source/TextureShare/Private/Game/ViewExtension/TextureShareSceneViewExtension.h
--- a/Engine/Plugins/VirtualProduction/TextureShare/Source/TextureShare/Private/Game/ViewExtension/TextureShareSceneViewExtension.h
+++ b/Engine/Plugins/VirtualProduction/TextureShare/Source/TextureShare/Private/Game/ViewExtension/TextureShareSceneViewExtension.h
@@ -5,6 +5,7 @@
 #include "Containers/TextureShareCoreEnums.h"
 
 class ITextureShareObjectProxy;
+class SWindow;
 struct FTextureShareCoreViewDesc;
 struct FTextureShareSceneView;
 
diff --git a/Engine/Plugins/VirtualProduction/TextureShare/Source/TextureShare/Private/Game/WorldSubsystem/TextureShareWorldSubsystemObjectProxy.cpp b/Engine/Plugins/VirtualProduction/TextureShare/Source/TextureShare/Private/Game/WorldSubsystem/TextureShareWorldSubsystemObjectProxy.cpp
--- a/Engine/Plugins/VirtualProduction/TextureShare/Source/TextureShare/Private/Game/WorldSubsystem/TextureShareWorldSubsystemObjectProxy.cpp
+++ b/Engine/Plugins/VirtualProduction/TextureShare/Source/TextureShare/Private/Game/WorldSubsystem/TextureShareWorldSubsystemObjectProxy.cpp
@@ -11,6 +11,7 @@
 #include "ITextureShareCoreObject.h"
 
 #include "Blueprints/TextureShareBlueprintContainers.h"
+#include "TextureResource.h"
 
 //////////////////////////////////////////////////////////////////////////////////////////////
 struct FReceiveTexture
diff --git a/Engine/Source/Runtime/Experimental/GeometryCollectionEngine/Private/GeometryCollection/GeometryCollectionObject.cpp b/Engine/Source/Runtime/Experimental/GeometryCollectionEngine/Private/GeometryCollection/GeometryCollectionObject.cpp
--- a/Engine/Source/Runtime/Experimental/GeometryCollectionEngine/Private/GeometryCollection/GeometryCollectionObject.cpp
+++ b/Engine/Source/Runtime/Experimental/GeometryCollectionEngine/Private/GeometryCollection/GeometryCollectionObject.cpp
@@ -6,6 +6,7 @@
 #include "GeometryCollection/GeometryCollectionObject.h"
 #include "GeometryCollection/GeometryCollection.h"
 #include "GeometryCollection/GeometryCollectionCache.h"
+#include "Materials/Material.h"
 #include "UObject/DestructionObjectVersion.h"
 #include "UObject/UE5MainStreamObjectVersion.h"
 #include "Serialization/ArchiveCountMem.h"
diff --git a/Engine/Source/Runtime/GameMenuBuilder/Private/SGameMenuPageWidget.cpp b/Engine/Source/Runtime/GameMenuBuilder/Private/SGameMenuPageWidget.cpp
--- a/Engine/Source/Runtime/GameMenuBuilder/Private/SGameMenuPageWidget.cpp
+++ b/Engine/Source/Runtime/GameMenuBuilder/Private/SGameMenuPageWidget.cpp
@@ -2,6 +2,7 @@
 
 #include "SGameMenuPageWidget.h"
 #include "Engine/GameViewportClient.h"
+#include "UnrealClient.h"
 #include "Widgets/SBoxPanel.h"
 #include "Widgets/SOverlay.h"
 #include "Framework/Application/SlateApplication.h"
